Control Rig / RigVM: rename .p4ignore

#rnx

#Synced-CL 51826256

diff --git a/Engine/Plugins/Animation/ControlRig/.p4ignore.txt b/Engine/Plugins/Animation/ControlRig/.p4ignore.txt
new file mode 100644
--- /dev/null
+++ b/Engine/Plugins/Animation/ControlRig/.p4ignore.txt
@@ -0,0 +1,7 @@
+.claude
+.git
+.gitignore
+Binaries
+Intermediate
+
+*.pyc
\ No newline at end of file
diff --git a/Engine/Plugins/Runtime/RigVM/.p4ignore.txt b/Engine/Plugins/Runtime/RigVM/.p4ignore.txt
new file mode 100644
--- /dev/null
+++ b/Engine/Plugins/Runtime/RigVM/.p4ignore.txt
@@ -0,0 +1,7 @@
+.claude
+.git
+.gitignore
+Binaries
+Intermediate
+
+*.pyc
\ No newline at end of file
